Added GPU-Reshape Plugin, handles local environment bootstrapping/injection
- Added PIX/RD style editor button to open the app and auto attaches the workspace

diff --git a/Engine/Plugins/Developer/GPUReshape/GPUReshape.uplugin b/Engine/Plugins/Developer/GPUReshape/GPUReshape.uplugin
new file mode 100644
--- /dev/null
+++ b/Engine/Plugins/Developer/GPUReshape/GPUReshape.uplugin
@@ -0,0 +1,22 @@
+{
+	"FileVersion": 1,
+	"Version": 1,
+	"VersionName": "0.1",
+	"FriendlyName": "GPU Reshape Plugin",
+	"Description": "GPU Reshape editor integration",
+	"Category": "Rendering",
+	"CreatedBy": "Epic Games, Inc.",
+	"CreatedByURL": "https://epicgames.com",
+	"DocsURL": "",
+	"EnabledByDefault": true,
+	"SupportedTargetPlatforms":  [ "Win64" ],
+	"Modules": [
+		{
+			"Name": "GPUReshape",
+			"Type": "DeveloperTool",
+			"LoadingPhase": "PostConfigInit",
+			"PlatformAllowList": [ "Win64" ],
+			"TargetDenyList": [ "Server" ]
+		}
+	]
+}
diff --git a/Engine/Plugins/Developer/GPUReshape/Source/GPUReshape/GPUReshape.Build.cs b/Engine/Plugins/Developer/GPUReshape/Source/GPUReshape/GPUReshape.Build.cs
new file mode 100644
--- /dev/null
+++ b/Engine/Plugins/Developer/GPUReshape/Source/GPUReshape/GPUReshape.Build.cs
@@ -0,0 +1,42 @@
+// Copyright Epic Games, Inc. All Rights Reserved.
+
+using EpicGames.Core;
+
+namespace UnrealBuildTool.Rules
+{
+	[SupportedPlatformGroups("Windows")]
+	public class GPUReshape : ModuleRules
+	{
+		public GPUReshape(ReadOnlyTargetRules Target) : base(Target)
+        {
+			PublicDependencyModuleNames.AddAll(
+				"InputDevice",
+				"RenderCore"
+			);
+
+			PrivateDependencyModuleNames.AddAll(
+				"Core",
+				"CoreUObject",
+				"Engine",
+				"InputCore",
+				"Projects",
+				"RHI"
+			);
+
+			if (Target.bBuildEditor)
+			{
+				DynamicallyLoadedModuleNames.Add("LevelEditor");
+
+				PrivateDependencyModuleNames.AddAll(
+					"Slate", 
+					"SlateCore",
+					"EditorFramework",
+					"UnrealEd",
+					"MainFrame",
+					"GameProjectGeneration",
+					"ToolMenus"
+				);
+			}
+		}
+	}
+}
diff --git a/Engine/Plugins/Developer/GPUReshape/Source/GPUReshape/Private/GPUReshapeCommands.cpp b/Engine/Plugins/Developer/GPUReshape/Source/GPUReshape/Private/GPUReshapeCommands.cpp
new file mode 100644
--- /dev/null
+++ b/Engine/Plugins/Developer/GPUReshape/Source/GPUReshape/Private/GPUReshapeCommands.cpp
@@ -0,0 +1,16 @@
+// Copyright Epic Games, Inc. All Rights Reserved.
+
+#if WITH_EDITOR
+
+#include "GPUReshapeCommands.h"
+
+#define LOCTEXT_NAMESPACE "GPUReshape"
+
+void FGPUReshapeCommands::RegisterCommands()
+{
+	UI_COMMAND(OpenApp, "Open GPU-Reshape", "Open the GPU-Reshape application for debugging/instrumentation (Alt+F12)", EUserInterfaceActionType::Button, FInputChord(EKeys::F12, EModifierKey::Alt));
+}
+
+#undef LOCTEXT_NAMESPACE
+
+#endif //WITH_EDITOR
diff --git a/Engine/Plugins/Developer/GPUReshape/Source/GPUReshape/Private/GPUReshapeCommands.h b/Engine/Plugins/Developer/GPUReshape/Source/GPUReshape/Private/GPUReshapeCommands.h
new file mode 100644
--- /dev/null
+++ b/Engine/Plugins/Developer/GPUReshape/Source/GPUReshape/Private/GPUReshapeCommands.h
@@ -0,0 +1,26 @@
+// Copyright Epic Games, Inc. All Rights Reserved.
+
+#pragma once
+
+#include "Misc/CoreMiscDefines.h"
+
+#if WITH_EDITOR
+
+#include "GPUReshapeStyle.h"
+#include "Framework/Commands/Commands.h"
+
+class FGPUReshapeCommands : public TCommands<FGPUReshapeCommands>
+{
+public:
+	FGPUReshapeCommands() : TCommands(TEXT("GPUReshape"), NSLOCTEXT("Contexts", "GPUReshape", "GPUReshape Plugin"), NAME_None, FGPUReshapeStyle::Get()->GetStyleSetName())
+	{
+		
+	}
+
+	/** TCommands overrides */
+	virtual void RegisterCommands() override;
+
+	TSharedPtr<class FUICommandInfo> OpenApp;
+};
+
+#endif // WITH_EDITOR
diff --git a/Engine/Plugins/Developer/GPUReshape/Source/GPUReshape/Private/GPUReshapeModule.cpp b/Engine/Plugins/Developer/GPUReshape/Source/GPUReshape/Private/GPUReshapeModule.cpp
new file mode 100644
--- /dev/null
+++ b/Engine/Plugins/Developer/GPUReshape/Source/GPUReshape/Private/GPUReshapeModule.cpp
@@ -0,0 +1,294 @@
+// Copyright Epic Games, Inc. All Rights Reserved.
+
+#include "GPUReshapeModule.h"
+#include "RenderingThread.h"
+#include "UnrealClient.h"
+#include "Engine/GameViewportClient.h"
+#include "Engine/Engine.h"
+#include "Misc/CommandLine.h"
+#include "Misc/Paths.h"
+#include "HAL/FileManager.h"
+#include "RHI.h"
+
+#if WITH_EDITOR
+#include "GPUReshapeCommands.h"
+#include "GPUReshapeStyle.h"
+#include "Framework/Application/SlateApplication.h"
+#include "Kismet2/DebuggerCommands.h"
+#include "ToolMenus.h"
+#include "Editor.h"
+#endif // WITH_EDITOR
+
+#if PLATFORM_WINDOWS
+#	include "Windows/WindowsHWrapper.h"
+#else // PLATFORM_WINDOWS
+#	error Not supported
+#endif // PLATFORM_WINDOWS
+
+DEFINE_LOG_CATEGORY(LogGPUReshape);
+
+static int32 GAutoAttachGPUReshape = 0;
+static FAutoConsoleVariableRef CVarAutoAttachGPUReshape(
+	TEXT("r.AutoAttachGPUReshape"),
+	GAutoAttachGPUReshape,
+	TEXT("Automatically loads GPU-Reshape on startup"),
+	ECVF_ReadOnly | ECVF_RenderThreadSafe
+);
+
+/** Manually hoisted from Loader.h, until we have proper include paths */
+namespace Detail
+{
+	using PFN_GRS_LOADER_INSTALL            = bool(*)();
+	using PFN_GRS_LOADER_GET_RESERVED_TOKEN = void(*)(char* output, uint32_t* length);
+	static constexpr const wchar_t* kPFNGRSLoaderInstallNameW = L"GRSLoaderInstall";
+	static constexpr const wchar_t* kPFNGRSLoaderGetReservedTokenNameW = L"GRSLoaderGetReservedToken";
+} // Detail
+
+void FGPUReshapeModule::StartupModule()
+{
+	const bool bCommandLineRequested =
+		FParse::Param(FCommandLine::Get(), TEXT("AttachGRS")) ||
+		FParse::Param(FCommandLine::Get(), TEXT("AttachReshape")) ||
+		FParse::Param(FCommandLine::Get(), TEXT("AttachGPUReshape"));
+	
+	// Requested reshape?
+	if (!GAutoAttachGPUReshape && !bCommandLineRequested)
+	{
+		return;
+	}
+
+	// Setup the loader and any backend state
+	if (!FindAndInstallLoader())
+	{
+		return;
+	}
+	
+	// Register console command
+	OpenAppCommand = MakeUnique<FAutoConsoleCommand>(
+		TEXT("GRS"), TEXT("Opens GPU Reshape"),
+		FConsoleCommandDelegate::CreateRaw(this, &FGPUReshapeModule::OpenOrSwitchToApp)
+	);
+
+	// Post extension install
+	FCoreDelegates::OnPostEngineInit.AddRaw(this, &FGPUReshapeModule::OnPostEngineInit);
+
+	bBackendInitialized = true;
+	UE_LOG(LogGPUReshape, Log, TEXT("Initialized"));
+}
+
+bool FGPUReshapeModule::InstallLoader(void* Handle)
+{
+	Detail::PFN_GRS_LOADER_INSTALL PFNInstall = reinterpret_cast<Detail::PFN_GRS_LOADER_INSTALL>(FPlatformProcess::GetDllExport(Handle, Detail::kPFNGRSLoaderInstallNameW));
+	if (!PFNInstall)
+	{
+		UE_LOG(LogGPUReshape, Error, TEXT("Failed to GPA '%s'"), Detail::kPFNGRSLoaderInstallNameW);
+		return false;
+	}
+
+	// Handles general bootstrapping, layer injections, etc.
+	return PFNInstall();
+}
+
+void FGPUReshapeModule::CacheLoaderReservedToken(void* Handle)
+{
+	Detail::PFN_GRS_LOADER_GET_RESERVED_TOKEN PFNGetReservedToken = reinterpret_cast<Detail::PFN_GRS_LOADER_GET_RESERVED_TOKEN>(FPlatformProcess::GetDllExport(Handle, Detail::kPFNGRSLoaderGetReservedTokenNameW));
+	if (!PFNGetReservedToken)
+	{
+		UE_LOG(LogGPUReshape, Error, TEXT("Failed to GPA '%s'"), Detail::kPFNGRSLoaderGetReservedTokenNameW);
+		return;
+	}
+
+	// Length of token
+	uint32 tokenLength;
+	PFNGetReservedToken(nullptr, &tokenLength);
+
+	// Get token data
+	char token[256]{};
+	PFNGetReservedToken(token, &tokenLength);
+
+	ReservedToken.Empty();
+	ReservedToken.Append(token, tokenLength);
+}
+
+void FGPUReshapeModule::ShutdownModule()
+{
+	OpenAppCommand.Reset();
+	
+#if WITH_EDITOR
+	FGPUReshapeStyle::Shutdown();
+	FGPUReshapeCommands::Unregister();
+#endif // WITH_EDITOR
+}
+
+bool FGPUReshapeModule::FindAndInstallLoader()
+{
+	if (GUsingNullRHI)
+	{
+		UE_LOG(LogGPUReshape, Error, TEXT("Null-RHI not supported"));
+		return false;
+	}
+	
+	// Allow binary overrides (useful for development)
+	if (!FParse::Value(FCommandLine::Get(), TEXT("GRSPath"), GPUReshapePath))
+	{
+		// Default path, for now at least
+		GPUReshapePath = FPaths::EngineDir() / TEXT("Restricted/NotForLicensees/Binaries/ThirdParty/GPUReshape/Win64");
+	}
+
+	// Expected loader, handles all backends
+	const FString GRSLoaderName = TEXT("GRS.Services.Loader.dll");
+
+	// Valid path?
+	const FString GRSLoaderPath = FPaths::Combine(*GPUReshapePath, *GRSLoaderName);
+	if (!FPaths::FileExists(GRSLoaderPath))
+	{
+		UE_LOG(LogGPUReshape, Error, TEXT("Loader not found: '%s'"), *GRSLoaderPath);
+		return false;
+	}
+
+	// Try to load the lib
+	void* Handle = FPlatformProcess::GetDllHandle(*GRSLoaderPath);
+	if (!Handle)
+	{
+		UE_LOG(LogGPUReshape, Error, TEXT("Failed to load library"));
+		return false;
+	}
+
+	// Try to install it
+	if (!InstallLoader(Handle))
+	{
+		UE_LOG(LogGPUReshape, Error, TEXT("Failed to install loader"));
+		return false;
+	}
+
+	// The initialized workspace is assigned a token, let's get it for later attaching
+	CacheLoaderReservedToken(Handle);
+	return true;
+}
+
+void FGPUReshapeModule::SwitchToApp()
+{
+	// Bring all relevant windows to the foreground
+	EnumWindows([](HWND Window, LPARAM Param) -> BOOL
+	{
+		FGPUReshapeModule* This = reinterpret_cast<FGPUReshapeModule*>(Param);
+
+		// Is child window?
+		DWORD RemoteProcessID = 0;
+		DWORD RemoteThreadID  = GetWindowThreadProcessId(Window, &RemoteProcessID);
+		if (!RemoteThreadID || RemoteProcessID != This->GetAppGetProcessID())
+		{
+			return true;
+		}
+
+		// Attach to avoid anti-focus stealing
+		AttachThreadInput(GetCurrentThreadId(), RemoteThreadID, true);
+		SetForegroundWindow(Window);
+		AttachThreadInput(GetCurrentThreadId(), RemoteThreadID, false);
+		return true;
+	}, reinterpret_cast<LPARAM>(this));
+}
+
+void FGPUReshapeModule::OpenOrSwitchToApp()
+{
+	// Already running?
+	if (AppProcHandle.IsValid() && FPlatformProcess::IsProcRunning(AppProcHandle))
+	{
+		SwitchToApp();
+		return;
+	}
+	
+	// Valid path?
+	const FString AppPath = FPaths::Combine(*GPUReshapePath, TEXT("GPUReshape.exe"));
+	if (!FPaths::FileExists(AppPath))
+	{
+		UE_LOG(LogGPUReshape, Error, TEXT("App not found: '%s'"), *AppPath);
+		return;
+	}
+
+	// Launch by attaching to the existing workspace
+	FString CommandLine;
+	CommandLine += "attach ";
+	CommandLine += FString::Printf(TEXT("-pid %u "), FPlatformProcess::GetCurrentProcessId());
+	CommandLine += FString::Printf(TEXT("-token %s "), *ReservedToken);
+	CommandLine += FString::Printf(TEXT("-process %s "), FPlatformProcess::ExecutableName(false));
+
+	// Try to create a job for shared lifetimes
+	HANDLE Job = CreateJobObject(nullptr, nullptr);
+	if (Job)
+	{
+		JOBOBJECT_EXTENDED_LIMIT_INFORMATION Info{};
+		Info.BasicLimitInformation.LimitFlags = JOB_OBJECT_LIMIT_KILL_ON_JOB_CLOSE;
+		if (!SetInformationJobObject(Job, JobObjectExtendedLimitInformation, &Info, sizeof(Info)))
+		{
+			UE_LOG(LogGPUReshape, Error, TEXT("Failed to initialize job"));
+		}
+	}
+	else
+	{
+		UE_LOG(LogGPUReshape, Error, TEXT("Failed to create job"));
+	}
+	
+	AppProcHandle = FPlatformProcess::CreateProc(
+		*AppPath,
+		*CommandLine,
+		false,
+		false,
+		false,
+		&AppProcessID,
+		0,
+		nullptr,
+		nullptr,
+		nullptr
+	);
+
+	// Group the process lifetimes
+	if (Job && !AssignProcessToJobObject(Job, AppProcHandle.Get()))
+	{
+		UE_LOG(LogGPUReshape, Error, TEXT("Failed to assign job"));
+	}
+}
+
+void FGPUReshapeModule::OnPostEngineInit()
+{
+#if WITH_EDITOR
+	if (!FSlateApplication::IsInitialized() || IsRunningCommandlet())
+	{
+		return;
+	}
+	
+	// Setup utilities
+	FGPUReshapeStyle::Initialize();
+	FGPUReshapeCommands::Register();
+
+	if (!IsRunningGame())
+	{
+		// Support the new viewport toolbar
+		if (FToolMenuOwnerScoped Scope(this); UToolMenu* Menu = UToolMenus::Get()->ExtendMenu("LevelEditor.ViewportToolbar"))
+		{
+			FToolMenuSection& Section = Menu->FindOrAddSection("Right");
+
+			// Add new section
+			FToolMenuEntry& Entry = Section.AddMenuEntry(FGPUReshapeCommands::Get().OpenApp);
+			Entry.ToolBarData.LabelOverride = FText::GetEmpty();
+			Entry.InsertPosition.Position = EToolMenuInsertType::First;	
+		}
+	}
+
+	if (GEditor)
+	{
+		check(FPlayWorldCommands::GlobalPlayWorldActions.IsValid());
+
+		FPlayWorldCommands::GlobalPlayWorldActions->MapAction(
+			FGPUReshapeCommands::Get().OpenApp,
+			FExecuteAction::CreateLambda([this]
+			{
+				OpenOrSwitchToApp();
+			}),
+			FCanExecuteAction()
+		);
+	}
+#endif // WITH_EDITOR
+}
+
+IMPLEMENT_MODULE(FGPUReshapeModule, GPUReshape)
diff --git a/Engine/Plugins/Developer/GPUReshape/Source/GPUReshape/Private/GPUReshapeModule.h b/Engine/Plugins/Developer/GPUReshape/Source/GPUReshape/Private/GPUReshapeModule.h
new file mode 100644
--- /dev/null
+++ b/Engine/Plugins/Developer/GPUReshape/Source/GPUReshape/Private/GPUReshapeModule.h
@@ -0,0 +1,71 @@
+// Copyright Epic Games, Inc. All Rights Reserved.
+
+#pragma once
+
+#include "Logging/LogMacros.h"
+#include "Modules/ModuleManager.h"
+#include "HAL/PlatformProcess.h"
+
+DECLARE_LOG_CATEGORY_EXTERN(LogGPUReshape, Log, All);
+
+class FGPUReshapeModule : public IModuleInterface
+{
+public:
+	/** Open, or switch to if present, the app */
+	void OpenOrSwitchToApp();
+
+	/** Switch to the already open app */
+	void SwitchToApp();
+
+	/** Get the current process id */
+	uint32 GetAppGetProcessID() const
+	{
+		return AppProcessID;
+	}
+
+	/** Is GPU Reshape initialized correctly? */
+	bool IsInitialized() const
+	{
+		return bBackendInitialized;
+	}
+	
+protected: /** IModuleInterface */
+	virtual void StartupModule() override;
+	virtual void ShutdownModule() override;
+
+private:
+	/** Find the appropriate loader and install it */
+	bool FindAndInstallLoader();
+	
+	/** Install the loader onto the current process */
+	bool InstallLoader(void* Handle);
+
+	/** Cache the reserved token for attaching  */
+	void CacheLoaderReservedToken(void* Handle);
+
+	/** Invoked for editor extensions */
+	void OnPostEngineInit();
+
+private:
+	/** Commands */
+	TUniquePtr<class FAutoConsoleCommand> OpenAppCommand;
+
+	/** Bin path for tooling */
+	FString GPUReshapePath;
+
+	/** Cached token for attaching */
+	FString ReservedToken;
+
+	/** Current process handle */
+	FProcHandle AppProcHandle;
+
+	/** Current process id */
+	uint32 AppProcessID = 0;
+
+	/** Everything set up right? */
+	bool bBackendInitialized = false;
+
+#if WITH_EDITOR
+	TSharedPtr<class FGPUReshapeEditorExtension> EditorExtension;
+#endif // WITH_EDITOR
+};
diff --git a/Engine/Plugins/Developer/GPUReshape/Source/GPUReshape/Private/GPUReshapeStyle.cpp b/Engine/Plugins/Developer/GPUReshape/Source/GPUReshape/Private/GPUReshapeStyle.cpp
new file mode 100644
--- /dev/null
+++ b/Engine/Plugins/Developer/GPUReshape/Source/GPUReshape/Private/GPUReshapeStyle.cpp
@@ -0,0 +1,69 @@
+// Copyright Epic Games, Inc. All Rights Reserved.
+
+#if WITH_EDITOR
+
+#include "GPUReshapeStyle.h"
+#include "Styling/SlateStyle.h"
+#include "Styling/SlateStyleRegistry.h"
+#include "HAL/FileManager.h"
+#include "Interfaces/IPluginManager.h"
+#include "GPUReshapeModule.h"
+
+static TSharedPtr<FSlateStyleSet> GPUReshapeStyleSet;
+
+/** Base path for all resources */
+static FString GetAssetResourcePath()
+{
+	TSharedPtr<IPlugin> ThisPlugin = IPluginManager::Get().FindPlugin(TEXT("GPUReshape"));
+	if (!ThisPlugin)
+	{
+		UE_LOG(LogGPUReshape, Error, TEXT("Failed to find plugin"));
+		return "";
+	}
+	
+	return ThisPlugin->GetBaseDir() / TEXT("Resources");
+}
+
+static FString GetAssetPath(const FString& RelativePath, const ANSICHAR* Extension)
+{
+	return GetAssetResourcePath() / RelativePath + Extension;
+}
+
+void FGPUReshapeStyle::Initialize()
+{
+	// Allow lazy init
+	if (GPUReshapeStyleSet.IsValid())
+	{
+		return;
+	}
+
+	GPUReshapeStyleSet = MakeShareable(new FSlateStyleSet("GPUReshapeStyle"));
+
+	// Always lives in the engine
+	FString EngineResourceDir = FPaths::EnginePluginsDir() / TEXT("GPUReshape") / TEXT("Resources");
+	GPUReshapeStyleSet->SetContentRoot(EngineResourceDir);
+	GPUReshapeStyleSet->SetCoreContentRoot(EngineResourceDir);
+
+	// Load resources
+	GPUReshapeStyleSet->Set("GPUReshape.Icon", new FSlateImageBrush(GetAssetPath("IconNoFrame", ".png"), FVector2D(40.0f, 40.0f)));
+	GPUReshapeStyleSet->Set("GPUReshape.OpenApp", new FSlateImageBrush(GetAssetPath("IconNoFrame16", ".png"), FVector2D(16.0f, 16.0f)));
+
+	// Register
+	FSlateStyleRegistry::RegisterSlateStyle(*GPUReshapeStyleSet.Get());
+};
+
+void FGPUReshapeStyle::Shutdown()
+{
+	if (GPUReshapeStyleSet.IsValid())
+	{
+		FSlateStyleRegistry::UnRegisterSlateStyle(*GPUReshapeStyleSet.Get());
+		GPUReshapeStyleSet.Reset();
+	}
+}
+
+TSharedPtr<ISlateStyle> FGPUReshapeStyle::Get()
+{
+	return GPUReshapeStyleSet;
+}
+
+#endif // WITH_EDITOR
diff --git a/Engine/Plugins/Developer/GPUReshape/Source/GPUReshape/Private/GPUReshapeStyle.h b/Engine/Plugins/Developer/GPUReshape/Source/GPUReshape/Private/GPUReshapeStyle.h
new file mode 100644
--- /dev/null
+++ b/Engine/Plugins/Developer/GPUReshape/Source/GPUReshape/Private/GPUReshapeStyle.h
@@ -0,0 +1,20 @@
+// Copyright Epic Games, Inc. All Rights Reserved.
+
+#pragma once
+
+#if WITH_EDITOR
+
+#include "Styling/ISlateStyle.h"
+
+class FGPUReshapeStyle
+{
+public:
+	/** Setup */
+	static void Initialize();
+	static void Shutdown();
+
+	/** Get the shared style */
+	static TSharedPtr<ISlateStyle> Get();
+};
+
+#endif // WITH_EDITOR
