Fixed bad option parsing for GPU Reshape branch

diff --git a/Engine/Plugins/Developer/GPUReshape/Source/GPUReshape/Private/GPUReshapeModule.cpp b/Engine/Plugins/Developer/GPUReshape/Source/GPUReshape/Private/GPUReshapeModule.cpp
--- a/Engine/Plugins/Developer/GPUReshape/Source/GPUReshape/Private/GPUReshapeModule.cpp
+++ b/Engine/Plugins/Developer/GPUReshape/Source/GPUReshape/Private/GPUReshapeModule.cpp
@@ -133,7 +133,7 @@ bool FGPUReshapeModule::FindAndInstallLoader()
 	{
 		// Default to the raytracing branch, until it's been merged back to dev
 		FString Branch = "Raytracing";
-		FParse::Value(FCommandLine::Get(), TEXT("GRSBranch"), GPUReshapePath);
+		FParse::Value(FCommandLine::Get(), TEXT("GRSBranch"), Branch);
 	
 		// Default path, for now at least
 		GPUReshapePath = FPaths::EngineDir() / TEXT("Restricted/NotForLicensees/Binaries/ThirdParty/GPUReshape/Branch") / Branch / TEXT("Win64");
