GPU Reshape, added automatic symbol paths

diff --git a/Engine/Plugins/Developer/GPUReshape/Source/GPUReshape/Private/GPUReshapeModule.cpp b/Engine/Plugins/Developer/GPUReshape/Source/GPUReshape/Private/GPUReshapeModule.cpp
--- a/Engine/Plugins/Developer/GPUReshape/Source/GPUReshape/Private/GPUReshapeModule.cpp
+++ b/Engine/Plugins/Developer/GPUReshape/Source/GPUReshape/Private/GPUReshapeModule.cpp
@@ -38,7 +38,17 @@ static FAutoConsoleVariableRef CVarAutoAttachGPUReshape(
 /** Manually hoisted from Loader.h, until we have proper include paths */
 namespace Detail
 {
-	using PFN_GRS_LOADER_INSTALL            = bool(*)();
+	struct GRSLoaderSymbolInfo {
+		const char** paths = nullptr;
+		uint32_t pathCount = 0;
+		bool includeSubDirectories = false;
+	};
+
+	struct GRSLoaderInstallInfo {
+		GRSLoaderSymbolInfo symbol;
+	};
+
+	using PFN_GRS_LOADER_INSTALL            = bool(*)(const GRSLoaderInstallInfo* info);
 	using PFN_GRS_LOADER_GET_RESERVED_TOKEN = void(*)(char* output, uint32_t* length);
 	static constexpr const wchar_t* kPFNGRSLoaderInstallNameW = L"GRSLoaderInstall";
 	static constexpr const wchar_t* kPFNGRSLoaderGetReservedTokenNameW = L"GRSLoaderGetReservedToken";
@@ -85,8 +95,18 @@ bool FGPUReshapeModule::InstallLoader(void* Handle)
 		return false;
 	}
 
+	// Symbol directory, Reshape is fairly Ansi based, for now at least
+	FAnsiString ProjectSymbolDir    = FAnsiString(FPaths::ProjectSavedDir() / TEXT("ShaderSymbols"));
+	const char* ProjectSymbolDirPtr = *ProjectSymbolDir;
+
+	// Setup the installation environment
+	Detail::GRSLoaderInstallInfo info{};
+	info.symbol.includeSubDirectories = true;
+	info.symbol.pathCount = 1;
+	info.symbol.paths = &ProjectSymbolDirPtr;
+
 	// Handles general bootstrapping, layer injections, etc.
-	return PFNInstall();
+	return PFNInstall(&info);
 }
 
 void FGPUReshapeModule::CacheLoaderReservedToken(void* Handle)
@@ -164,6 +184,8 @@ bool FGPUReshapeModule::FindAndInstallLoader()
 		UE_LOG(LogGPUReshape, Error, TEXT("Failed to install loader"));
 		return false;
 	}
+	
+	UE_LOG(LogGPUReshape, Display, TEXT("Installed loader: '%s'"), *GRSLoaderPath);
 
 	// The initialized workspace is assigned a token, let's get it for later attaching
 	CacheLoaderReservedToken(Handle);
