GPU Reshape [addressing feedback], automatically set symbol and source paths

#Synced-CL 51539352

diff --git a/Engine/Plugins/Developer/GPUReshape/Source/GPUReshape/GPUReshape.Build.cs b/Engine/Plugins/Developer/GPUReshape/Source/GPUReshape/GPUReshape.Build.cs
--- a/Engine/Plugins/Developer/GPUReshape/Source/GPUReshape/GPUReshape.Build.cs
+++ b/Engine/Plugins/Developer/GPUReshape/Source/GPUReshape/GPUReshape.Build.cs
@@ -23,7 +23,8 @@ namespace UnrealBuildTool.Rules
 				"Engine",
 				"InputCore",
 				"Projects",
-				"RHI"
+				"RHI",
+				"Json"
 			);
 
 			string[] SourcePaths = [
diff --git a/Engine/Plugins/Developer/GPUReshape/Source/GPUReshape/Private/GPUReshapeModule.cpp b/Engine/Plugins/Developer/GPUReshape/Source/GPUReshape/Private/GPUReshapeModule.cpp
--- a/Engine/Plugins/Developer/GPUReshape/Source/GPUReshape/Private/GPUReshapeModule.cpp
+++ b/Engine/Plugins/Developer/GPUReshape/Source/GPUReshape/Private/GPUReshapeModule.cpp
@@ -9,6 +9,9 @@
 #include "Misc/Paths.h"
 #include "HAL/FileManager.h"
 #include "RHI.h"
+#include "Dom/JsonObject.h"
+#include "Serialization/JsonSerializer.h"
+#include "Misc/FileHelper.h"
 
 #if WITH_EDITOR
 #include "GPUReshapeCommands.h"
@@ -199,8 +202,50 @@ void FGPUReshapeModule::SwitchToApp()
 	}, reinterpret_cast<LPARAM>(this));
 }
 
+static bool CreateProjectSettingJson(FString& OutPath)
+{
+	TSharedPtr<FJsonObject> Object = MakeShared<FJsonObject>();
+	
+	// Targeting the editor
+	Object->SetStringField(TEXT("ApplicationName"), TEXT("UnrealEditor.exe"));
+	
+	// Setup symbols
+	TSharedPtr<FJsonObject> SymbolObject = MakeShared<FJsonObject>();
+	{
+		SymbolObject->SetArrayField(TEXT("Paths"), {
+			MakeShared<FJsonValueString>(FPaths::ConvertRelativePathToFull(FPaths::ProjectSavedDir() / TEXT("ShaderSymbols")))
+		});
+		
+		// Always sub-index
+		SymbolObject->SetBoolField(TEXT("IncludeSubDirectories"), true);
+	}
+	Object->SetObjectField(TEXT("Symbol"), SymbolObject);
+	
+	// Setup symbols
+	TSharedPtr<FJsonObject> SourceObject = MakeShared<FJsonObject>();
+	{
+		SourceObject->SetArrayField(TEXT("Paths"), {
+			MakeShared<FJsonValueString>(FPaths::ConvertRelativePathToFull(FPaths::EngineDir() / TEXT("Shaders")))
+		});
+		
+		// Always sub-index
+		SourceObject->SetBoolField(TEXT("IncludeSubDirectories"), true);
+	}
+	Object->SetObjectField(TEXT("Source"), SourceObject);
+	
+	// Serialize to text
+	FString Contents;
+	FJsonSerializer::Serialize(Object.ToSharedRef(), TJsonWriterFactory<>::Create(&Contents));
+	
+	// Try to write out
+	OutPath = FPaths::CreateTempFilename(FPlatformProcess::UserTempDir(), TEXT("GPUReshapeSettings-"), TEXT(".json"));
+	return FFileHelper::SaveStringToFile(Contents, *OutPath);
+} 
+
 void FGPUReshapeModule::OpenOrSwitchToApp()
 {
+	FString CommandLine;
+	
 	// Already running?
 	if (AppProcHandle.IsValid() && FPlatformProcess::IsProcRunning(AppProcHandle))
 	{
@@ -215,9 +260,16 @@ void FGPUReshapeModule::OpenOrSwitchToApp()
 		UE_LOG(LogGPUReshape, Error, TEXT("App not found: '%s'"), *AppPath);
 		return;
 	}
-
+	
+	// Attach default settings
+	if (FString SettingPath; CreateProjectSettingJson(SettingPath))
+	{
+		CommandLine += "apply-setting ";
+		CommandLine += FString::Printf(TEXT("-json \"%s\" "), *SettingPath);
+		CommandLine += " && ";
+	}
+	
 	// Launch by attaching to the existing workspace
-	FString CommandLine;
 	CommandLine += "attach ";
 	CommandLine += FString::Printf(TEXT("-pid %u "), FPlatformProcess::GetCurrentProcessId());
 	CommandLine += FString::Printf(TEXT("-token %s "), *ReservedToken);
